UVEditor: Avoid loading content browser items when trying to figure out whether UV editor can be started on them. Fix the example asset editor implementation as well.

diff --git a/Engine/Plugins/Experimental/CharacterFXEditor/ExampleCharacterFXEditor/Source/ExampleCharacterFXEditor/Private/ExampleCharacterFXEditorModule.cpp b/Engine/Plugins/Experimental/CharacterFXEditor/ExampleCharacterFXEditor/Source/ExampleCharacterFXEditor/Private/ExampleCharacterFXEditorModule.cpp
--- a/Engine/Plugins/Experimental/CharacterFXEditor/ExampleCharacterFXEditor/Source/ExampleCharacterFXEditor/Private/ExampleCharacterFXEditorModule.cpp
+++ b/Engine/Plugins/Experimental/CharacterFXEditor/ExampleCharacterFXEditor/Source/ExampleCharacterFXEditor/Private/ExampleCharacterFXEditorModule.cpp
@@ -55,19 +55,30 @@ void FExampleCharacterFXEditorModule::RegisterMenus()
 				// We'll need to get the target objects out of the context
 				if (UContentBrowserAssetContextMenuContext* Context = Section.FindContext<UContentBrowserAssetContextMenuContext>())
 				{
-					TArray<TObjectPtr<UObject>> AssetsToEdit;
-					AssetsToEdit.Append(Context->GetSelectedObjects());
-
 					UExampleCharacterFXEditorSubsystem* ExampleCharacterFXEditorSubsystem = GEditor->GetEditorSubsystem<UExampleCharacterFXEditorSubsystem>();
 					check(ExampleCharacterFXEditorSubsystem);
 
-					if (ExampleCharacterFXEditorSubsystem->AreObjectsValidTargets(AssetsToEdit))
+					// We are deliberately not using Context->GetSelectedObjects() here to avoid triggering a load from right clicking
+					// an asset in the content browser.
+					if (ExampleCharacterFXEditorSubsystem->AreAssetsValidTargets(Context->SelectedAssets))
 					{
 						TSharedPtr<class FUICommandList> CommandListToBind = MakeShared<FUICommandList>();
 
 						CommandListToBind->MapAction(
 							FExampleCharacterFXEditorCommands::Get().OpenCharacterFXEditor,
-							FExecuteAction::CreateUObject(ExampleCharacterFXEditorSubsystem, &UExampleCharacterFXEditorSubsystem::StartExampleCharacterFXEditor, AssetsToEdit));
+							FExecuteAction::CreateWeakLambda(ExampleCharacterFXEditorSubsystem, [Context, ExampleCharacterFXEditorSubsystem]()
+							{
+								// When we actually do want to open the editor, trigger the load to get the objects
+								TArray<TObjectPtr<UObject>> AssetsToEdit;
+								AssetsToEdit.Append(Context->LoadSelectedObjects<UObject>());
+
+								// If we fail the ensure here, then there must be something that we're failing to check properly
+								// in AreAssetsValidTargets that we would need to track down and check in the asset data.
+								if (ensure(ExampleCharacterFXEditorSubsystem->AreObjectsValidTargets(AssetsToEdit)))
+								{
+									ExampleCharacterFXEditorSubsystem->StartExampleCharacterFXEditor(AssetsToEdit);
+								}
+							}));
 
 						Section.AddMenuEntryWithCommandList(FExampleCharacterFXEditorCommands::Get().OpenCharacterFXEditor,
 							CommandListToBind, 
diff --git a/Engine/Plugins/Experimental/CharacterFXEditor/ExampleCharacterFXEditor/Source/ExampleCharacterFXEditor/Private/ExampleCharacterFXEditorSubsystem.cpp b/Engine/Plugins/Experimental/CharacterFXEditor/ExampleCharacterFXEditor/Source/ExampleCharacterFXEditor/Private/ExampleCharacterFXEditorSubsystem.cpp
--- a/Engine/Plugins/Experimental/CharacterFXEditor/ExampleCharacterFXEditor/Source/ExampleCharacterFXEditor/Private/ExampleCharacterFXEditorSubsystem.cpp
+++ b/Engine/Plugins/Experimental/CharacterFXEditor/ExampleCharacterFXEditor/Source/ExampleCharacterFXEditor/Private/ExampleCharacterFXEditorSubsystem.cpp
@@ -73,6 +73,40 @@ bool UExampleCharacterFXEditorSubsystem::AreObjectsValidTargets(const TArray<UOb
 	return true;
 }
 
+bool UExampleCharacterFXEditorSubsystem::AreAssetsValidTargets(const TArray<FAssetData>& InAssets) const
+{
+	if (InAssets.IsEmpty())
+	{
+		return false;
+	}
+
+	// Currently our tool target factories don't evaluate FAssetData to figure out whether they can
+	// build a tool target (they only work on UObjects directly), so for now we do corresponding checks
+	// here ourselves according to the tooltargets that we support.
+	auto IsValidStaticMeshAsset = [](const FAssetData& AssetData)
+	{
+		// The static mesh tool target checks GetNumSourceModels, which we can't do directly, hence our check of the LODs tag
+		int32 NumLODs = 0;
+		return AssetData.IsInstanceOf<UStaticMesh>() && AssetData.GetTagValue<int32>("LODs", NumLODs) && NumLODs > 0;
+	};
+	auto IsValidSkeletalMeshAsset = [](const FAssetData& AssetData)
+	{
+		// The skeletal mesh tool targets don't seem to try to check the number of LODs, but the skeletal mesh tool target
+		// uses an exact cast for the class, hence the '==' comparison here.
+		return AssetData.GetClass() == USkeletalMesh::StaticClass();
+	};
+
+	for (const FAssetData& AssetData : InAssets)
+	{
+		if (!IsValidStaticMeshAsset(AssetData) && !IsValidSkeletalMeshAsset(AssetData))
+		{
+			return false;
+		}
+	}
+
+	return true;
+}
+
 void UExampleCharacterFXEditorSubsystem::StartExampleCharacterFXEditor(TArray<TObjectPtr<UObject>> ObjectsToEdit)
 {
 	// We don't allow opening a new instance if any of the objects are already opened
@@ -94,6 +128,13 @@ void UExampleCharacterFXEditorSubsystem::StartExampleCharacterFXEditor(TArray<TO
 	// If we got here, there's not an instance already opened.
 	UExampleCharacterFXEditor* CharacterFXEditor = NewObject<UExampleCharacterFXEditor>();
 
+	// We should have done a check upstream to make sure that all of our targets are valid, but
+	// we'll check again here.
+	if (!ensure(AreObjectsValidTargets(ObjectsToEdit)))
+	{
+		return;
+	}
+
 	// Among other things, this call registers the editor with the asset editor subsystem,
 	// which will prevent it from being garbage collected.
 	CharacterFXEditor->Initialize(ObjectsToEdit);
diff --git a/Engine/Plugins/Experimental/CharacterFXEditor/ExampleCharacterFXEditor/Source/ExampleCharacterFXEditor/Public/ExampleCharacterFXEditorSubsystem.h b/Engine/Plugins/Experimental/CharacterFXEditor/ExampleCharacterFXEditor/Source/ExampleCharacterFXEditor/Public/ExampleCharacterFXEditorSubsystem.h
--- a/Engine/Plugins/Experimental/CharacterFXEditor/ExampleCharacterFXEditor/Source/ExampleCharacterFXEditor/Public/ExampleCharacterFXEditorSubsystem.h
+++ b/Engine/Plugins/Experimental/CharacterFXEditor/ExampleCharacterFXEditor/Source/ExampleCharacterFXEditor/Public/ExampleCharacterFXEditorSubsystem.h
@@ -41,6 +41,14 @@ public:
 	/** Checks that all of the objects are valid targets for an editor session. */
 	virtual bool AreObjectsValidTargets(const TArray<UObject*>& InObjects) const;
 
+	/**
+	 * Checks that all of the assets are valid targets for an editor session. This
+	 * is preferable over AreObjectsValidTargets when we have FAssetData because it
+	 * allows us to avoid forcing a load of the underlying UObjects (for instance to
+	 * avoid triggering a load when right clicking an asset in the content browser).
+	 */
+	virtual bool AreAssetsValidTargets(const TArray<FAssetData>& InAssets) const;
+
 	/**
 	 * Either brings to the front an existing editor instance that is editing one of
 	 * these objects, if one exists, or starts up a new instance editing all of these
diff --git a/Engine/Plugins/Experimental/UVEditor/Source/UVEditor/Private/UVEditorModule.cpp b/Engine/Plugins/Experimental/UVEditor/Source/UVEditor/Private/UVEditorModule.cpp
--- a/Engine/Plugins/Experimental/UVEditor/Source/UVEditor/Private/UVEditorModule.cpp
+++ b/Engine/Plugins/Experimental/UVEditor/Source/UVEditor/Private/UVEditorModule.cpp
@@ -82,21 +82,32 @@ void FUVEditorModule::RegisterMenus()
 		Section.AddDynamicEntry("OpenUVEditor", FNewToolMenuSectionDelegate::CreateLambda(
 			[](FToolMenuSection& Section)
 			{
-				// We'll need to get the target objects out of the context
+				// We'll need to get the target assets out of the context
 				if (UContentBrowserAssetContextMenuContext* Context = Section.FindContext<UContentBrowserAssetContextMenuContext>())
 				{
-					TArray<TObjectPtr<UObject>> AssetsToEdit;
-					AssetsToEdit.Append(Context->GetSelectedObjects());
-
 					UUVEditorSubsystem* UVSubsystem = GEditor->GetEditorSubsystem<UUVEditorSubsystem>();
 					check(UVSubsystem);
 
-					if (UVSubsystem->AreObjectsValidTargets(AssetsToEdit))
+					// We are deliberately not using Context->GetSelectedObjects() here to avoid triggering a load from right clicking
+					// an asset in the content browser.
+					if (UVSubsystem->AreAssetsValidTargets(Context->SelectedAssets))
 					{
 						TSharedPtr<class FUICommandList> CommandListToBind = MakeShared<FUICommandList>();
 						CommandListToBind->MapAction(
 							FUVEditorCommands::Get().OpenUVEditor,
-							FExecuteAction::CreateUObject(UVSubsystem, &UUVEditorSubsystem::StartUVEditor, AssetsToEdit));
+							FExecuteAction::CreateWeakLambda(UVSubsystem, [Context, UVSubsystem]()
+							{
+								// When we actually do want to open the editor, trigger the load to get the objects
+								TArray<TObjectPtr<UObject>> AssetsToEdit;
+								AssetsToEdit.Append(Context->LoadSelectedObjects<UObject>());
+
+								// If we fail the ensure here, then there must be something that we're failing to check properly
+								// in AreAssetsValidTargets that we would need to track down and check in the asset data.
+								if (ensure(UVSubsystem->AreObjectsValidTargets(AssetsToEdit)))
+								{
+									UVSubsystem->StartUVEditor(AssetsToEdit);
+								}
+							}));
 
 						Section.AddMenuEntryWithCommandList(FUVEditorCommands::Get().OpenUVEditor, CommandListToBind, TAttribute<FText>(), TAttribute<FText>(), FSlateIcon(FUVEditorStyle::Get().GetStyleSetName(), "UVEditor.OpenUVEditor"));
 					}
diff --git a/Engine/Plugins/Experimental/UVEditor/Source/UVEditor/Private/UVEditorSubsystem.cpp b/Engine/Plugins/Experimental/UVEditor/Source/UVEditor/Private/UVEditorSubsystem.cpp
--- a/Engine/Plugins/Experimental/UVEditor/Source/UVEditor/Private/UVEditorSubsystem.cpp
+++ b/Engine/Plugins/Experimental/UVEditor/Source/UVEditor/Private/UVEditorSubsystem.cpp
@@ -49,6 +49,40 @@ bool UUVEditorSubsystem::AreObjectsValidTargets(const TArray<UObject*>& InObject
 	return true;
 }
 
+bool UUVEditorSubsystem::AreAssetsValidTargets(const TArray<FAssetData>& InAssets) const
+{
+	if (InAssets.IsEmpty())
+	{
+		return false;
+	}
+
+	// Currently our tool target factories don't evaluate FAssetData to figure out whether they can
+	// build a tool target (they only work on UObjects directly), so for now we do corresponding checks
+	// here ourselves.
+	auto IsValidStaticMeshAsset = [](const FAssetData& AssetData)
+	{
+		// The static mesh tool target checks GetNumSourceModels, which we can't do directly, hence our check of the LODs tag
+		int32 NumLODs = 0;
+		return AssetData.IsInstanceOf<UStaticMesh>() && AssetData.GetTagValue<int32>("LODs", NumLODs) && NumLODs > 0;
+	};
+	auto IsValidSkeletalMeshAsset = [](const FAssetData& AssetData)
+	{
+		// The skeletal mesh tool targets don't seem to try to check the number of LODs, but the skeletal mesh tool target
+		// uses an exact cast for the class, hence the '==' comparison here.
+		return AssetData.GetClass() == USkeletalMesh::StaticClass();
+	};
+
+	for (const FAssetData& AssetData : InAssets)
+	{
+		if (!IsValidStaticMeshAsset(AssetData) && !IsValidSkeletalMeshAsset(AssetData))
+		{
+			return false;
+		}
+	}
+
+	return true;
+}
+
 void UUVEditorSubsystem::BuildTargets(const TArray<TObjectPtr<UObject>>& ObjectsIn, 
 	const FToolTargetTypeRequirements& TargetRequirements, TArray<TObjectPtr<UToolTarget>>& TargetsOut)
 {
@@ -83,6 +117,13 @@ void UUVEditorSubsystem::StartUVEditor(TArray<TObjectPtr<UObject>> ObjectsToEdit
 
 	// If we got here, there's not an instance already opened.
 
+	// We should have done a check upstream to make sure that all of our targets are valid, but
+	// we'll check again here.
+	if (!ensure(AreObjectsValidTargets(ObjectsToEdit)))
+	{
+		return;
+	}
+
 	UUVEditor* UVEditor = NewObject<UUVEditor>(this);
 
 	// Among other things, this call registers the UV editor with the asset editor subsystem,
diff --git a/Engine/Plugins/Experimental/UVEditor/Source/UVEditor/Public/UVEditorSubsystem.h b/Engine/Plugins/Experimental/UVEditor/Source/UVEditor/Public/UVEditorSubsystem.h
--- a/Engine/Plugins/Experimental/UVEditor/Source/UVEditor/Public/UVEditorSubsystem.h
+++ b/Engine/Plugins/Experimental/UVEditor/Source/UVEditor/Public/UVEditorSubsystem.h
@@ -37,6 +37,14 @@ public:
 	/** Checks that all of the objects are valid targets for a UV editor session. */
 	virtual bool AreObjectsValidTargets(const TArray<UObject*>& InObjects) const;
 
+	/**
+	 * Checks that all of the assets are valid targets for an editor session. This
+	 * is preferable over AreObjectsValidTargets when we have FAssetData because it
+	 * allows us to avoid forcing a load of the underlying UObjects (for instance to
+	 * avoid triggering a load when right clicking an asset in the content browser).
+	 */
+	virtual bool AreAssetsValidTargets(const TArray<FAssetData>& InAssets) const;
+
 	/** 
 	 * Tries to build the core targets that provide meshes for UV tools to work on.
 	 */
